Fix a bug of missing mobile shaders when using Performance mode.

diff --git a/Engine/Plugins/Runtime/WindowsDeviceProfileSelector/Source/WindowsDeviceProfileSelector/Private/WindowsDeviceProfileSelectorModule.cpp b/Engine/Plugins/Runtime/WindowsDeviceProfileSelector/Source/WindowsDeviceProfileSelector/Private/WindowsDeviceProfileSelectorModule.cpp
--- a/Engine/Plugins/Runtime/WindowsDeviceProfileSelector/Source/WindowsDeviceProfileSelector/Private/WindowsDeviceProfileSelectorModule.cpp
+++ b/Engine/Plugins/Runtime/WindowsDeviceProfileSelector/Source/WindowsDeviceProfileSelector/Private/WindowsDeviceProfileSelectorModule.cpp
@@ -26,11 +26,8 @@ const FString FWindowsDeviceProfileSelectorModule::GetRuntimeDeviceProfileName()
 
 	if (FApp::CanEverRender())
 	{
-		FString DeviceProfileFileName;
-		FConfigCacheIni::LoadGlobalIniFile(DeviceProfileFileName, TEXT("DeviceProfiles"));
-
 		TArray<FString> AvailableProfiles;
-		GConfig->GetSectionNames(DeviceProfileFileName, AvailableProfiles);
+		GConfig->GetSectionNames(GDeviceProfilesIni, AvailableProfiles);
 
 		FString TmpProfileName = ProfileName + TCHAR('_') + GetSelectedDynamicRHIModuleName(false);
 		if (AvailableProfiles.Contains(FString::Printf(TEXT("%s DeviceProfile"), *TmpProfileName)))
