Windows-based platform extensions can utilize Windows device profiles.

#rnx

diff --git a/Engine/Plugins/Runtime/WindowsDeviceProfileSelector/Source/WindowsDeviceProfileSelector/Private/WindowsDeviceProfileSelectorModule.cpp b/Engine/Plugins/Runtime/WindowsDeviceProfileSelector/Source/WindowsDeviceProfileSelector/Private/WindowsDeviceProfileSelectorModule.cpp
--- a/Engine/Plugins/Runtime/WindowsDeviceProfileSelector/Source/WindowsDeviceProfileSelector/Private/WindowsDeviceProfileSelectorModule.cpp
+++ b/Engine/Plugins/Runtime/WindowsDeviceProfileSelector/Source/WindowsDeviceProfileSelector/Private/WindowsDeviceProfileSelectorModule.cpp
@@ -21,8 +21,11 @@ void FWindowsDeviceProfileSelectorModule::ShutdownModule()
 
 const FString FWindowsDeviceProfileSelectorModule::GetRuntimeDeviceProfileName()
 {
+#if defined(WINDOWS_OVERRIDE_DEVICEPROFILE_NAME)
+	// overridden device profile name, typically for Windows-based platform extensions
+	FString ProfileName = WINDOWS_OVERRIDE_DEVICEPROFILE_NAME;
+#elif UE_IS_COOKED_EDITOR
 	// some heuristics to determine a cooked editor
-#if UE_IS_COOKED_EDITOR
 	FString ProfileName = TEXT("WindowsCookedEditor");
 #else
 	// Windows, WindowsEditor, WindowsClient, or WindowsServer
