Remove old settings migration code which was incorrectly firing. Also reduces RHI name spam.

diff --git a/Engine/Plugins/Runtime/WindowsDeviceProfileSelector/Source/WindowsDeviceProfileSelector/Private/WindowsDeviceProfileSelectorModule.cpp b/Engine/Plugins/Runtime/WindowsDeviceProfileSelector/Source/WindowsDeviceProfileSelector/Private/WindowsDeviceProfileSelectorModule.cpp
--- a/Engine/Plugins/Runtime/WindowsDeviceProfileSelector/Source/WindowsDeviceProfileSelector/Private/WindowsDeviceProfileSelectorModule.cpp
+++ b/Engine/Plugins/Runtime/WindowsDeviceProfileSelector/Source/WindowsDeviceProfileSelector/Private/WindowsDeviceProfileSelectorModule.cpp
@@ -45,6 +45,5 @@ const FString FWindowsDeviceProfileSelectorModule::GetRuntimeDeviceProfileName()
 		}
 	}
 
-	UE_LOG(LogInit, Log, TEXT("Selected Device Profile: [%s]"), *ProfileName);
 	return ProfileName;
 }
diff --git a/Engine/Source/Runtime/Engine/Private/DeviceProfiles/DeviceProfileManager.cpp b/Engine/Source/Runtime/Engine/Private/DeviceProfiles/DeviceProfileManager.cpp
--- a/Engine/Source/Runtime/Engine/Private/DeviceProfiles/DeviceProfileManager.cpp
+++ b/Engine/Source/Runtime/Engine/Private/DeviceProfiles/DeviceProfileManager.cpp
@@ -527,6 +527,8 @@ void UDeviceProfileManager::InitializeCVarsForActiveDeviceProfile(bool bPushSett
 		ActiveProfileName = GetPlatformDeviceProfileName();
 	}
 
+	UE_LOG(LogInit, Log, TEXT("Selected Device Profile: [%s]"), *ActiveProfileName);
+
 	SetDeviceProfileCVars(ActiveProfileName);
 }
 
