Cache GConfig reads from FWindowsDeviceProfileSelectorModule to reduce potential hitches if GEngineIni gets unloaded

#tests Editor, Windows Client
#rnx
#Synced-CL 49641272

diff --git a/Engine/Plugins/Runtime/WindowsDeviceProfileSelector/Source/WindowsDeviceProfileSelector/Private/WindowsDeviceProfileSelectorModule.cpp b/Engine/Plugins/Runtime/WindowsDeviceProfileSelector/Source/WindowsDeviceProfileSelector/Private/WindowsDeviceProfileSelectorModule.cpp
--- a/Engine/Plugins/Runtime/WindowsDeviceProfileSelector/Source/WindowsDeviceProfileSelector/Private/WindowsDeviceProfileSelectorModule.cpp
+++ b/Engine/Plugins/Runtime/WindowsDeviceProfileSelector/Source/WindowsDeviceProfileSelector/Private/WindowsDeviceProfileSelectorModule.cpp
@@ -22,36 +22,38 @@ void FWindowsDeviceProfileSelectorModule::ShutdownModule()
 
 const FString FWindowsDeviceProfileSelectorModule::GetRuntimeDeviceProfileName()
 {
+	static FString ProfileName;
+
+	if (ProfileName.IsEmpty())
+	{
 #if defined(WINDOWS_OVERRIDE_DEVICEPROFILE_NAME)
-	// overridden device profile name, typically for Windows-based platform extensions
-	FString ProfileName = WINDOWS_OVERRIDE_DEVICEPROFILE_NAME;
+		// overridden device profile name, typically for Windows-based platform extensions
+		ProfileName = WINDOWS_OVERRIDE_DEVICEPROFILE_NAME;
 #elif UE_IS_COOKED_EDITOR
-	// some heuristics to determine a cooked editor
-	FString ProfileName = TEXT("WindowsCookedEditor");
+		// some heuristics to determine a cooked editor
+		ProfileName = TEXT("WindowsCookedEditor");
 #else
-	// Windows, WindowsEditor, WindowsClient, or WindowsServer
-	FString ProfileName = FPlatformProperties::PlatformName();
+		// Windows, WindowsEditor, WindowsClient, or WindowsServer
+		ProfileName = FPlatformProperties::PlatformName();
 #endif
 
-	if (FApp::CanEverRender())
-	{
-		TArray<FString> AvailableProfiles;
-		GConfig->GetSectionNames(GDeviceProfilesIni, AvailableProfiles);
-
-		FString TmpProfileName = ProfileName + TCHAR('_') + GetSelectedDynamicRHIModuleName(false);
-		if (AvailableProfiles.Contains(FString::Printf(TEXT("%s DeviceProfile"), *TmpProfileName)))
+		if (FApp::CanEverRender())
 		{
-			// Use RHI specific device profile if exist
-			ProfileName = MoveTemp(TmpProfileName);
-		}
-		else if (TmpProfileName.Contains(TEXT("_ES31")))
-		{
-			TmpProfileName = ProfileName + TEXT("_ES31");
-			if (AvailableProfiles.Contains(FString::Printf(TEXT("%s DeviceProfile"), *TmpProfileName)))
+			FString TmpProfileName = ProfileName + TCHAR('_') + GetSelectedDynamicRHIModuleName(false);
+			if (GConfig->DoesSectionExist(*FString::Printf(TEXT("%s DeviceProfile"), *TmpProfileName), GDeviceProfilesIni))
 			{
-				// General ES31 device profile
+				// Use RHI specific device profile if exist
 				ProfileName = MoveTemp(TmpProfileName);
 			}
+			else if (TmpProfileName.Contains(TEXT("_ES31")))
+			{
+				TmpProfileName = ProfileName + TEXT("_ES31");
+				if (GConfig->DoesSectionExist(*FString::Printf(TEXT("%s DeviceProfile"), *TmpProfileName), GDeviceProfilesIni))
+				{
+					// General ES31 device profile
+					ProfileName = MoveTemp(TmpProfileName);
+				}
+			}
 		}
 	}
 
