Speed up InitializeCVarsForActiveDeviceProfile by skipping GetSelectedDynamicRHIModuleName when no RHI-specific device profiles exist.
GetSelectedDynamicRHIModuleName calls LoadDynamicRHIModule, which takes ~3s, but was called unconditionally whenever FApp::CanEverRender() returned true, even when no RHI-specific device profiles were configured. This change first checks if there are any RHI-specific device profiles that could apply, before calling GetSelectedDynamicRHIModuleName.
This usually moves LoadDynamicRHIModule to a later stage in the engine/editor loading sequence, which makes the splashscreen show up faster.

#rnx
#Synced-CL 52172106

diff --git a/Engine/Plugins/Runtime/WindowsDeviceProfileSelector/Source/WindowsDeviceProfileSelector/Private/WindowsDeviceProfileSelectorModule.cpp b/Engine/Plugins/Runtime/WindowsDeviceProfileSelector/Source/WindowsDeviceProfileSelector/Private/WindowsDeviceProfileSelectorModule.cpp
--- a/Engine/Plugins/Runtime/WindowsDeviceProfileSelector/Source/WindowsDeviceProfileSelector/Private/WindowsDeviceProfileSelectorModule.cpp
+++ b/Engine/Plugins/Runtime/WindowsDeviceProfileSelector/Source/WindowsDeviceProfileSelector/Private/WindowsDeviceProfileSelectorModule.cpp
@@ -7,10 +7,44 @@
 #include "DynamicRHI.h"
 #include "Misc/App.h"
 #include "Misc/ConfigCacheIni.h"
+#include "Internationalization/Regex.h"
 
 IMPLEMENT_MODULE(FWindowsDeviceProfileSelectorModule, WindowsDeviceProfileSelector);
 
 
+static TSet<FString> GetRHIDeviceProfileNames(const FString& BaseProfileName)
+{
+	// BaseProfileName is embedded directly into a regex pattern, so it must not contain
+	// any regex special characters. Profile names are always plain platform name strings
+	// (e.g. "Windows", "WindowsEditor"), so this would indicate a programming error.
+	for (TCHAR C : BaseProfileName)
+	{
+		checkf(FChar::IsAlnum(C) || C == TCHAR('_'),
+			TEXT("BaseProfileName '%s' contains unexpected character '%c'"), *BaseProfileName, C);
+	}
+
+	TSet<FString> RHIProfileNames;
+	const FConfigFile* DeviceProfilesConfig = GConfig->Find(GDeviceProfilesIni);
+	if (DeviceProfilesConfig == nullptr)
+	{
+		return RHIProfileNames;
+	}
+
+	const FRegexPattern RHIProfilePattern(
+		FString::Printf(TEXT("^%s_(\\w+) DeviceProfile$"), *BaseProfileName));
+	const int32 DeviceProfileSuffixLen = FString(TEXT(" DeviceProfile")).Len();
+	for (const auto& [SectionName, Section] : *DeviceProfilesConfig)
+	{
+		FRegexMatcher Matcher(RHIProfilePattern, SectionName);
+		if (Matcher.FindNext())
+		{
+			RHIProfileNames.Add(SectionName.Left(SectionName.Len() - DeviceProfileSuffixLen));
+		}
+	}
+	return RHIProfileNames;
+}
+
+
 void FWindowsDeviceProfileSelectorModule::StartupModule()
 {
 }
@@ -41,20 +75,27 @@ const FString FWindowsDeviceProfileSelectorModule::GetRuntimeDeviceProfileName()
 
 		if (FApp::CanEverRender())
 		{
-			FString TmpProfileName = ProfileName + TCHAR('_') + GetSelectedDynamicRHIModuleName(false);
-			if (GConfig->DoesSectionExist(*FString::Printf(TEXT("%s DeviceProfile"), *TmpProfileName), GDeviceProfilesIni))
-			{
-				// Use RHI specific device profile if exist
-				ProfileName = MoveTemp(TmpProfileName);
-			}
-			else if (TmpProfileName.Contains(TEXT("_ES31")))
+			// Scan all section names to find RHI-specific device profiles for this base profile name.
+			// This avoids calling the expensive GetSelectedDynamicRHIModuleName when no such profiles exist.
+			TSet<FString> ConfigRHIProfileNames = GetRHIDeviceProfileNames(ProfileName);
+
+			if (ConfigRHIProfileNames.Num() > 0)
 			{
-				TmpProfileName = ProfileName + TEXT("_ES31");
-				if (GConfig->DoesSectionExist(*FString::Printf(TEXT("%s DeviceProfile"), *TmpProfileName), GDeviceProfilesIni))
+				FString TmpProfileName = ProfileName + TCHAR('_') + GetSelectedDynamicRHIModuleName(false);
+				if (ConfigRHIProfileNames.Contains(TmpProfileName))
 				{
-					// General ES31 device profile
+					// Use RHI specific device profile if it exists
 					ProfileName = MoveTemp(TmpProfileName);
 				}
+				else if (TmpProfileName.Contains(TEXT("_ES31")))
+				{
+					TmpProfileName = ProfileName + TEXT("_ES31");
+					if (ConfigRHIProfileNames.Contains(TmpProfileName))
+					{
+						// General ES31 device profile
+						ProfileName = MoveTemp(TmpProfileName);
+					}
+				}
 			}
 		}
 	}
diff --git a/Engine/Source/Runtime/D3D12RHI/Private/Windows/WindowsD3D12Device.cpp b/Engine/Source/Runtime/D3D12RHI/Private/Windows/WindowsD3D12Device.cpp
--- a/Engine/Source/Runtime/D3D12RHI/Private/Windows/WindowsD3D12Device.cpp
+++ b/Engine/Source/Runtime/D3D12RHI/Private/Windows/WindowsD3D12Device.cpp
@@ -757,6 +757,8 @@ static bool CheckIfAgilitySDKLoaded()
 
 bool FD3D12DynamicRHIModule::IsSupported(ERHIFeatureLevel::Type RequestedFeatureLevel)
 {
+	TRACE_CPUPROFILER_EVENT_SCOPE(FD3D12DynamicRHIModule::IsSupported);
+
 	// Windows version 15063 is Windows 1703 aka "Windows Creator Update"
 	// This is the first version that supports ID3D12Device2 which is our minimum runtime device version.
 	if (!FPlatformMisc::VerifyWindowsVersion(10, 0, 15063))
diff --git a/Engine/Source/Runtime/Launch/Private/LaunchEngineLoop.cpp b/Engine/Source/Runtime/Launch/Private/LaunchEngineLoop.cpp
--- a/Engine/Source/Runtime/Launch/Private/LaunchEngineLoop.cpp
+++ b/Engine/Source/Runtime/Launch/Private/LaunchEngineLoop.cpp
@@ -3124,7 +3124,7 @@ int32 FEngineLoop::PreInitPreStartupScreen(const TCHAR* CmdLine)
 
 	FEmbeddedCommunication::ForceTick(3);
 
-	SlowTask.EnterProgressFrame(10);
+	SlowTask.EnterProgressFrame(5);
 
 #if USE_LOCALIZED_PACKAGE_CACHE
 	FPackageLocalizationManager::Get().InitializeFromLazyCallback([](FPackageLocalizationManager& InPackageLocalizationManager)
@@ -3154,6 +3154,7 @@ int32 FEngineLoop::PreInitPreStartupScreen(const TCHAR* CmdLine)
 	}
 
 	FDelayedAutoRegisterHelper::RunAndClearDelayedAutoRegisterDelegates(EDelayedRegisterRunPhase::PreRHIInit);
+	SlowTask.EnterProgressFrame(5, LOCTEXT("InitializingRHI", "Initializing Render Hardware Interface..."));
 	{
 		SCOPED_BOOT_TIMING("RHIInit");
 		// Initialize the RHI.
diff --git a/Engine/Source/Runtime/RHI/Private/Windows/WindowsDynamicRHI.cpp b/Engine/Source/Runtime/RHI/Private/Windows/WindowsDynamicRHI.cpp
--- a/Engine/Source/Runtime/RHI/Private/Windows/WindowsDynamicRHI.cpp
+++ b/Engine/Source/Runtime/RHI/Private/Windows/WindowsDynamicRHI.cpp
@@ -928,6 +928,7 @@ static TOptional<ERHIFeatureLevel::Type> ChooseDefaultFeatureLevel(EWindowsRHI I
 //  * D3D11, if all else fails.
 static EWindowsRHI ChooseDefaultRHI(const FParsedWindowsDynamicRHIConfig& Config)
 {
+	TRACE_CPUPROFILER_EVENT_SCOPE(ChooseDefaultRHI);
 	// Default graphics RHI is the main project setting that governs the choice, so it takes the priority
 	if (TOptional<EWindowsRHI> ConfigDefault = Config.DefaultRHI)
 	{
